Fix Monolithic build introduced by LiveLink

diff --git a/Engine/Source/Editor/SerializedRecorderInterface/Private/MovieSceneSectionSerialization.cpp b/Engine/Source/Editor/SerializedRecorderInterface/Private/MovieSceneSectionSerialization.cpp
--- a/Engine/Source/Editor/SerializedRecorderInterface/Private/MovieSceneSectionSerialization.cpp
+++ b/Engine/Source/Editor/SerializedRecorderInterface/Private/MovieSceneSectionSerialization.cpp
@@ -18,7 +18,8 @@ DEFINE_LOG_CATEGORY(MovieSceneSerialization);
 
 
 const FGuid FTempCustomVersion::GUID(0xCB8AB0CD, 0xE78C4BDE, 0xA8621393, 0x14E9EF62);
-FCustomVersionRegistration GRegisterLiveLinkCustomVersion(FTempCustomVersion::GUID, FTempCustomVersion::LatestVersion, TEXT("LiveLinkCustomVersion"));
+FCustomVersionRegistration GRegisterLiveLinkMovieSectionCustomVersion(FTempCustomVersion::GUID, FTempCustomVersion::LatestVersion, TEXT("LiveLinkCustomVersion"));
+
 const int64 MovieSceneSerializationNamespace::InvalidOffset = -1;
 const float MovieSceneSerializationNamespace::SerializerSleepTime = 0.2f;
 bool MovieSceneSerializationNamespace::bAutoSerialize = false;
diff --git a/Engine/Source/Runtime/LiveLinkInterface/Private/LiveLinkCustomVersion.cpp b/Engine/Source/Runtime/LiveLinkInterface/Private/LiveLinkCustomVersion.cpp
--- a/Engine/Source/Runtime/LiveLinkInterface/Private/LiveLinkCustomVersion.cpp
+++ b/Engine/Source/Runtime/LiveLinkInterface/Private/LiveLinkCustomVersion.cpp
@@ -6,4 +6,4 @@
 const FGuid FLiveLinkCustomVersion::GUID(0xab965196, 0x45d808fc, 0xb7d7228d, 0x78ad569e);
 
 // Register the custom version with core
-FCustomVersionRegistration GRegisterLiveLinkCustomVersion(FLiveLinkCustomVersion::GUID, FLiveLinkCustomVersion::LatestVersion, TEXT("LiveLinkVer"));
+FCustomVersionRegistration GRegisterLiveLinkInterfaceCustomVersion(FLiveLinkCustomVersion::GUID, FLiveLinkCustomVersion::LatestVersion, TEXT("LiveLinkVer"));
