Sequencer: Fix truncation warnings in MovieSceneTransformSerialization.h

#rnx

diff --git a/Engine/Source/Editor/SerializedRecorderInterface/Public/Serializers/MovieSceneTransformSerialization.h b/Engine/Source/Editor/SerializedRecorderInterface/Public/Serializers/MovieSceneTransformSerialization.h
--- a/Engine/Source/Editor/SerializedRecorderInterface/Public/Serializers/MovieSceneTransformSerialization.h
+++ b/Engine/Source/Editor/SerializedRecorderInterface/Public/Serializers/MovieSceneTransformSerialization.h
@@ -84,7 +84,7 @@ struct FSerializedTransform
 	}
 
 	FFrameNumber Time;
-	float Values[9];   //location, rotation, scale.
+	double Values[9];   //location, rotation, scale.
 };
 
 using FTransformSerializedFrame = TMovieSceneSerializedFrame<FSerializedTransform>;
